Cine Camera Capture Component: Added a new experimental plugin that allows users to capture Cine Camera results into Render target.

diff --git a/Engine/Plugins/Experimental/CineCameraSceneCapture/CineCameraSceneCapture.uplugin b/Engine/Plugins/Experimental/CineCameraSceneCapture/CineCameraSceneCapture.uplugin
new file mode 100644
--- /dev/null
+++ b/Engine/Plugins/Experimental/CineCameraSceneCapture/CineCameraSceneCapture.uplugin
@@ -0,0 +1,25 @@
+{
+	"FileVersion": 3,
+	"Version": 1,
+	"VersionName": "1.0",
+	"FriendlyName": "CineCameraSceneCapture",
+	"Description": "This plugin adds the ability to render cine camera views into Render Targets identically to Scene Capture",
+	"Category": "Virtual Production",
+	"CreatedBy": "Epic Games, Inc.",
+	"CreatedByURL": "https://epicgames.com",
+	"DocsURL": "",
+	"MarketplaceURL": "",
+	"SupportURL": "",
+	"EnabledByDefault": false,
+	"CanContainContent": true,
+	"IsBetaVersion": false,
+	"IsExperimentalVersion": true,
+	"Installed": false,
+	"Modules": [
+		{
+			"Name": "CineCameraSceneCapture",
+			"Type": "Runtime",
+			"LoadingPhase": "Default"
+		}
+	]
+}
\ No newline at end of file
diff --git a/Engine/Plugins/Experimental/CineCameraSceneCapture/Source/CineCameraSceneCapture/CineCameraSceneCapture.Build.cs b/Engine/Plugins/Experimental/CineCameraSceneCapture/Source/CineCameraSceneCapture/CineCameraSceneCapture.Build.cs
new file mode 100644
--- /dev/null
+++ b/Engine/Plugins/Experimental/CineCameraSceneCapture/Source/CineCameraSceneCapture/CineCameraSceneCapture.Build.cs
@@ -0,0 +1,21 @@
+// Copyright Epic Games, Inc. All Rights Reserved.
+
+using UnrealBuildTool;
+
+public class CineCameraSceneCapture : ModuleRules
+{
+	public CineCameraSceneCapture(ReadOnlyTargetRules Target) : base(Target)
+	{
+		PrivateDependencyModuleNames.AddRange(
+			new string[]
+			{
+				"Core",
+				"CoreUObject",
+				"Engine",
+				"Slate",
+				"SlateCore",
+				"CinematicCamera",
+			}
+		);
+	}
+}
diff --git a/Engine/Plugins/Experimental/CineCameraSceneCapture/Source/CineCameraSceneCapture/Private/CineCameraSceneCapture.cpp b/Engine/Plugins/Experimental/CineCameraSceneCapture/Source/CineCameraSceneCapture/Private/CineCameraSceneCapture.cpp
new file mode 100644
--- /dev/null
+++ b/Engine/Plugins/Experimental/CineCameraSceneCapture/Source/CineCameraSceneCapture/Private/CineCameraSceneCapture.cpp
@@ -0,0 +1,17 @@
+// Copyright Epic Games, Inc. All Rights Reserved.
+
+#include "CineCameraSceneCapture.h"
+
+#define LOCTEXT_NAMESPACE "FCineCameraSceneCaptureModule"
+
+void FCineCameraSceneCaptureModule::StartupModule()
+{
+}
+
+void FCineCameraSceneCaptureModule::ShutdownModule()
+{
+}
+
+#undef LOCTEXT_NAMESPACE
+	
+IMPLEMENT_MODULE(FCineCameraSceneCaptureModule, CineCameraSceneCapture)
\ No newline at end of file
diff --git a/Engine/Plugins/Experimental/CineCameraSceneCapture/Source/CineCameraSceneCapture/Private/CineCameraSceneCaptureComponent.cpp b/Engine/Plugins/Experimental/CineCameraSceneCapture/Source/CineCameraSceneCapture/Private/CineCameraSceneCaptureComponent.cpp
new file mode 100644
--- /dev/null
+++ b/Engine/Plugins/Experimental/CineCameraSceneCapture/Source/CineCameraSceneCapture/Private/CineCameraSceneCaptureComponent.cpp
@@ -0,0 +1,201 @@
+// Copyright Epic Games, Inc. All Rights Reserved.
+
+#include "CineCameraSceneCaptureComponent.h"
+#include "Components/StaticMeshComponent.h"
+#include "SceneViewExtension.h"
+#include "Runtime/CinematicCamera/Public/CineCameraComponent.h"
+#include "Runtime/Engine/Public/SceneView.h"
+#include "Engine/TextureRenderTarget2D.h"
+#include "Framework/Notifications/NotificationManager.h"
+#include "Widgets/Notifications/SNotificationList.h"
+
+#define LOCTEXT_NAMESPACE "FCineCameraSceneCaptureModule"
+
+DEFINE_LOG_CATEGORY(LogCineCapture);
+
+#define CINE_CAMERA_INVALID_PARENT_WARNING "Cine Capture requires to be parented to Cine Camera Component. Cine Capture {0} on Actor \"{1}\" will be disabled until it is parented to Cine Camera Actor."
+
+
+/** A helper class that assits switching setup state of CineCapture SVE for the purposes of setting up views for Cine Capture. */
+class FScopedCineCaptureSVESetup
+{
+public:
+	FScopedCineCaptureSVESetup(TSharedPtr<FCineCameraCaptureSceneViewExtension, ESPMode::ThreadSafe> InCineCaptureSVE, TSoftObjectPtr<UCineCameraComponent> InCineCameraComponent, bool bInFollowSceneCaptureRenderPath)
+	{
+		InCineCaptureSVE->OnViewSetupStarted(InCineCameraComponent, bInFollowSceneCaptureRenderPath);
+		CineCaptureSVE = InCineCaptureSVE;
+	}
+
+	~FScopedCineCaptureSVESetup()
+	{
+		CineCaptureSVE->OnViewSetupFinished();
+	}
+private:
+	TSharedPtr<FCineCameraCaptureSceneViewExtension, ESPMode::ThreadSafe> CineCaptureSVE;
+};
+
+
+UCineCaptureComponent2D::UCineCaptureComponent2D(const FObjectInitializer& ObjectInitializer)
+	: Super(ObjectInitializer), RenderTargetHighestDimension(1280), bFollowSceneCaptureRenderPath(true)
+{
+	CaptureSource = ESceneCaptureSource::SCS_FinalToneCurveHDR;
+	bAlwaysPersistRenderingState = true;
+}
+
+void UCineCaptureComponent2D::UpdateSceneCaptureContents(FSceneInterface* Scene)
+{
+	if (!CineCameraComponent.IsValid() || !CineCaptureSVE.IsValid())
+	{
+		return;
+	}
+
+	SetRelativeTransform(FTransform::Identity);
+	
+	PostProcessSettings = CineCameraComponent->PostProcessSettings;
+
+#if WITH_EDITORONLY_DATA
+	FOVAngle = CineCameraComponent->CurrentHorizontalFOV;
+#endif
+
+	bOverride_CustomNearClippingPlane = CineCameraComponent->bOverride_CustomNearClippingPlane;
+	CustomNearClippingPlane = CineCameraComponent->CustomNearClippingPlane;
+
+	{
+		/* Since SVE setup is happening on game thread, this guarantees that all SVEs unrelated to CineCapture will be excluded from modification. */
+		FScopedCineCaptureSVESetup ScopedCineCaptureSVESetup(CineCaptureSVE, CineCameraComponent, bFollowSceneCaptureRenderPath);
+		Scene->UpdateSceneCaptureContents(this);
+	}
+}
+
+void UCineCaptureComponent2D::CheckResizeRenderTarget()
+{
+	// Making sure that censor's aspect ratio is reflected render target and that RenderTarget Dimension specified by the user is taken into account.
+	if (TextureTarget)
+	{
+		float Ratio = TextureTarget->GetSurfaceWidth() / TextureTarget->GetSurfaceHeight();
+		const float SmallNum = 1e-3;
+		if ((FMath::Abs(Ratio - CineCameraComponent->AspectRatio) > SmallNum)
+			|| (FMath::Max(TextureTarget->GetSurfaceWidth(), TextureTarget->GetSurfaceHeight()) != RenderTargetHighestDimension))
+		{
+			if (CineCameraComponent->AspectRatio >= 1.)
+			{
+				int32 LowestDim = FMath::Max(RenderTargetHighestDimension / CineCameraComponent->AspectRatio, 1);
+				TextureTarget->ResizeTarget(RenderTargetHighestDimension, LowestDim);
+			}
+			else
+			{
+				int32 LowestDim = FMath::Max(RenderTargetHighestDimension * CineCameraComponent->AspectRatio, 1);
+				TextureTarget->ResizeTarget(LowestDim, RenderTargetHighestDimension);
+			}
+		}
+	}
+}
+
+void UCineCaptureComponent2D::OnRegister()
+{
+	Super::OnRegister();
+
+#if WITH_EDITORONLY_DATA
+	// Remove mesh created by Scene Capture Component.
+	ProxyMeshComponent->DestroyComponent();
+	ProxyMeshComponent = nullptr;
+	if (!CineCaptureSVE.IsValid())
+	{
+		// This will create a new scene extension per each capture.
+		CineCaptureSVE = FSceneViewExtensions::NewExtension<FCineCameraCaptureSceneViewExtension>();
+	}
+#endif
+	ValidateCineCameraComponent();
+}
+
+void UCineCaptureComponent2D::OnUnregister()
+{
+	Super::OnUnregister();
+	if (!CineCaptureSVE.IsValid())
+	{
+		CineCaptureSVE.Reset();
+	}
+	CineCameraComponent = nullptr;
+}
+
+void UCineCaptureComponent2D::OnAttachmentChanged()
+{
+	// If no parent is present, then this component is in a transient state.
+	if (USceneComponent* ParentComponent = GetAttachParent())
+	{
+		ValidateCineCameraComponent();
+	}
+}
+
+void UCineCaptureComponent2D::TickComponent(float DeltaTime, ELevelTick TickType, FActorComponentTickFunction* ThisTickFunction)
+{
+	Super::TickComponent(DeltaTime, TickType, ThisTickFunction);
+
+	if (CineCaptureSVE.IsValid())
+	{
+		CineCaptureSVE->SetFrameDeltaTime(DeltaTime);
+	}
+
+	// Since CineCameraComponent properties can change any time, render target needs to be validated every frame.
+	// The rest of the properties are copied right before the view is setup.
+	if (CineCameraComponent.IsValid())
+	{
+		CheckResizeRenderTarget();
+	}
+}
+
+void UCineCaptureComponent2D::ValidateCineCameraComponent()
+{
+	if (!IsValid(this))
+	{
+		return;
+	}
+
+	CineCameraComponent = Cast<UCineCameraComponent>(GetAttachParent());
+
+	if (!CineCameraComponent)
+	{
+#if WITH_EDITOR
+		FNotificationInfo Info(FText::Format(LOCTEXT("AddCineCameraNotification", CINE_CAMERA_INVALID_PARENT_WARNING), FText::FromString(GetName()), FText::FromString(GetOuter()->GetName())));
+		Info.ExpireDuration = 5.0f;
+
+		FSlateNotificationManager::Get().AddNotification(Info);
+#endif
+		UE_LOG(LogCineCapture, Warning, TEXT(CINE_CAMERA_INVALID_PARENT_WARNING), *GetName(), *GetOuter()->GetName());
+	}
+}
+
+FCineCameraCaptureSceneViewExtension::FCineCameraCaptureSceneViewExtension(const FAutoRegister& AutoRegister)
+	: FSceneViewExtensionBase(AutoRegister), bSetupViewInProcess(false)
+{
+
+}
+
+void FCineCameraCaptureSceneViewExtension::SetupView(FSceneViewFamily& InViewFamily, FSceneView& InView)
+{
+	if (!CineCameraComponent.IsValid())
+	{
+		return;
+	}
+
+	if (bSetupViewInProcess)
+	{
+		FMinimalViewInfo DesiredView;
+
+		// Immitate the behaviour of viewports.
+		CineCameraComponent->GetCameraView(DeltaTime, DesiredView);
+		FTransform Transform = CineCameraComponent->GetComponentToWorld();
+		FVector ViewLocation = Transform.GetTranslation();
+		InView.StartFinalPostprocessSettings(ViewLocation);
+
+		InView.OverridePostProcessSettings(DesiredView.PostProcessSettings, 1.0);
+		FSceneViewInitOptions ViewInitOptions;
+		InView.EndFinalPostprocessSettings(ViewInitOptions);
+
+		// Required for certain effects (lighting) to match that of Cine Camera.
+		InView.bIsSceneCapture = bFollowSceneCaptureRenderPath;
+	}
+}
+
+
+#undef LOCTEXT_NAMESPACE
diff --git a/Engine/Plugins/Experimental/CineCameraSceneCapture/Source/CineCameraSceneCapture/Public/CineCameraSceneCapture.h b/Engine/Plugins/Experimental/CineCameraSceneCapture/Source/CineCameraSceneCapture/Public/CineCameraSceneCapture.h
new file mode 100644
--- /dev/null
+++ b/Engine/Plugins/Experimental/CineCameraSceneCapture/Source/CineCameraSceneCapture/Public/CineCameraSceneCapture.h
@@ -0,0 +1,15 @@
+// Copyright Epic Games, Inc. All Rights Reserved.
+
+#pragma once
+
+#include "CoreMinimal.h"
+#include "Modules/ModuleManager.h"
+
+class FCineCameraSceneCaptureModule : public IModuleInterface
+{
+public:
+
+	/** IModuleInterface implementation */
+	virtual void StartupModule() override;
+	virtual void ShutdownModule() override;
+};
diff --git a/Engine/Plugins/Experimental/CineCameraSceneCapture/Source/CineCameraSceneCapture/Public/CineCameraSceneCaptureComponent.h b/Engine/Plugins/Experimental/CineCameraSceneCapture/Source/CineCameraSceneCapture/Public/CineCameraSceneCaptureComponent.h
new file mode 100644
--- /dev/null
+++ b/Engine/Plugins/Experimental/CineCameraSceneCapture/Source/CineCameraSceneCapture/Public/CineCameraSceneCaptureComponent.h
@@ -0,0 +1,114 @@
+// Copyright Epic Games, Inc. All Rights Reserved.
+
+#pragma once
+
+#include "Components/SceneCaptureComponent2D.h"
+#include "SceneViewExtension.h"
+#include "CineCameraSceneCaptureComponent.generated.h"
+
+DECLARE_LOG_CATEGORY_EXTERN(LogCineCapture, Log, All);
+
+class UCineCameraComponent;
+
+class FCineCameraCaptureSceneViewExtension : public FSceneViewExtensionBase
+{
+public:
+	FCineCameraCaptureSceneViewExtension(const FAutoRegister& AutoRegister);
+
+	//~ Begin FSceneViewExtensionBase Interface
+	virtual void SetupViewFamily(FSceneViewFamily& InViewFamily) override {};
+	virtual void SetupView(FSceneViewFamily& InViewFamily, FSceneView& InView) override;
+	virtual void BeginRenderViewFamily(FSceneViewFamily& InViewFamily) override {};
+	//~ End FSceneViewExtensionBase Interface
+public:
+	/** 
+	* Flags all views that will be setup until (bSetupViewInProcess == false) as related to Cine Capture.
+	*/
+	void OnViewSetupStarted(TSoftObjectPtr<UCineCameraComponent> InCineCameraComponent, bool bInFollowSceneCaptureRenderPath)
+	{ 
+		bSetupViewInProcess = true; 
+		CineCameraComponent = InCineCameraComponent;
+		bFollowSceneCaptureRenderPath = bInFollowSceneCaptureRenderPath;
+	};
+
+	/**
+	* Flags all following view setups as not related to Cine Capture.
+	*/
+	void OnViewSetupFinished() { bSetupViewInProcess = false; };
+
+	/**
+	* Delta time is needed for the purposes of camera smoothing.
+	*/
+	void SetFrameDeltaTime(float InDeltaTime)
+	{
+		DeltaTime = InDeltaTime;
+	}
+private:
+	/** 
+	* Translates to View.bIsSceneCapture.
+	*/
+	bool bFollowSceneCaptureRenderPath;
+
+	/** 
+	* When this property is true, it means that the views that are being setup currently are related to Cine Capture.
+	*/
+	bool bSetupViewInProcess;
+
+	/**
+	* A transient property that is used to deliver settings from CineCameraComponent to views that are related to cine capture.
+	*/
+	TSoftObjectPtr<UCineCameraComponent> CineCameraComponent;
+
+	/** 
+	* Delta time between frames. Used for camera smoothing.
+	*/
+	float DeltaTime;
+};
+
+/**
+* Cine Capture Component extends Scene Capture to allow users to render Cine Camera Component into a render target. 
+* Cine Capture has a few modifiable properties, but most of the properties are controlled by Cine Camera Component.
+* Cine Capture Component is required to be parented to Cine Camera Component or a class that extends it.
+* 
+*/
+UCLASS(hidecategories = (Transform, Collision, Object, Physics, SceneComponent, PostProcessVolume, Projection, Rendering, PlanarReflection), ClassGroup = Rendering, editinlinenew, meta = (BlueprintSpawnableComponent))
+class UCineCaptureComponent2D : public USceneCaptureComponent2D
+{
+	GENERATED_UCLASS_BODY()
+public:
+	
+	/** 
+	* Highest possible dimension of specified render target in either X or Y (depends on cine camera aspect ratio). Used to adjust the render target size.
+	* Example: With cine camera's censor aspect ratio is 2:1, and Highest Dimension set to 1000, the render target will be resized to 1000x500. 
+	*/
+	UPROPERTY(EditAnywhere, BlueprintReadWrite, Category = "Capture Settings", meta = (ClampMin = "1", ClampMax = "8192", UIMin = "1", UIMax = "8192"))
+	int32 RenderTargetHighestDimension;
+
+	/** 
+	* Affects rendering path cine capture takes. Scene Capture takes a slightly different rendering route compared to viewport rendering
+	* for the purposes of optimization. If the results of cine capture are vastly different to what cine camera displays try disabling this.
+	*/
+	UPROPERTY(EditAnywhere, BlueprintReadWrite, Category = "Capture Settings", AdvancedDisplay)
+	bool bFollowSceneCaptureRenderPath;
+
+	UPROPERTY()
+	TSoftObjectPtr<UCineCameraComponent> CineCameraComponent;
+
+public:
+	void CheckResizeRenderTarget();
+public:
+	virtual void UpdateSceneCaptureContents(FSceneInterface* Scene) override;
+	virtual void OnRegister() override;
+	virtual void OnUnregister() override;
+
+	/** Used for validation of this componet's attachment to Cine Camera Component. */
+	virtual void OnAttachmentChanged() override; 
+	virtual void TickComponent(float DeltaTime, enum ELevelTick TickType, FActorComponentTickFunction* ThisTickFunction) override;
+
+public:
+	void ValidateCineCameraComponent();
+
+private:
+	/** This scene view extension is used to get ahold of views during the setup process. */
+	TSharedPtr<FCineCameraCaptureSceneViewExtension, ESPMode::ThreadSafe> CineCaptureSVE;
+};
