CineCameraCapture: Further simplifications using the existing scene capture specific SVE.

#rnx

diff --git a/Engine/Plugins/Experimental/CineCameraSceneCapture/Source/CineCameraSceneCapture/Private/CineCameraSceneCaptureComponent.cpp b/Engine/Plugins/Experimental/CineCameraSceneCapture/Source/CineCameraSceneCapture/Private/CineCameraSceneCaptureComponent.cpp
--- a/Engine/Plugins/Experimental/CineCameraSceneCapture/Source/CineCameraSceneCapture/Private/CineCameraSceneCaptureComponent.cpp
+++ b/Engine/Plugins/Experimental/CineCameraSceneCapture/Source/CineCameraSceneCapture/Private/CineCameraSceneCaptureComponent.cpp
@@ -15,26 +15,70 @@ DEFINE_LOG_CATEGORY(LogCineCapture);
 
 #define CINE_CAMERA_INVALID_PARENT_WARNING "Cine Capture requires to be parented to Cine Camera Component. Cine Capture {0} on Actor \"{1}\" will be disabled until it is parented to Cine Camera Actor."
 
-
-/** A helper class that assits switching setup state of CineCapture SVE for the purposes of setting up views for Cine Capture. */
-class FScopedCineCaptureSVESetup
+// This extension is only registered onto the scene capture 2d component, and therefore runs locally.
+class FCineCameraCaptureSceneViewExtension : public ISceneViewExtension
 {
 public:
-	FScopedCineCaptureSVESetup(TSharedPtr<FCineCameraCaptureSceneViewExtension, ESPMode::ThreadSafe> InCineCaptureSVE, TSoftObjectPtr<UCineCameraComponent> InCineCameraComponent, bool bInFollowSceneCaptureRenderPath)
+	//~ Begin FSceneViewExtensionBase Interface
+	virtual void SetupViewFamily(FSceneViewFamily& InViewFamily) override {};
+	virtual void SetupView(FSceneViewFamily& InViewFamily, FSceneView& InView) override
 	{
-		InCineCaptureSVE->OnViewSetupStarted(InCineCameraComponent, bInFollowSceneCaptureRenderPath);
-		CineCaptureSVE = InCineCaptureSVE;
+		if (!CineCameraComponent.IsValid())
+		{
+			return;
+		}
+
+		FMinimalViewInfo DesiredView;
+
+		// Immitate the behaviour of viewports.
+		CineCameraComponent->GetCameraView(DeltaTime, DesiredView);
+		FTransform Transform = CineCameraComponent->GetComponentToWorld();
+		FVector ViewLocation = Transform.GetTranslation();
+		InView.StartFinalPostprocessSettings(ViewLocation);
+
+		InView.OverridePostProcessSettings(DesiredView.PostProcessSettings, 1.0);
+		FSceneViewInitOptions ViewInitOptions;
+		InView.EndFinalPostprocessSettings(ViewInitOptions);
+
+		// Required for certain effects (lighting) to match that of Cine Camera.
+		InView.bIsSceneCapture = bFollowSceneCaptureRenderPath;
 	}
+	virtual void BeginRenderViewFamily(FSceneViewFamily& InViewFamily) override {};
+	//~ End FSceneViewExtensionBase Interface
 
-	~FScopedCineCaptureSVESetup()
+	/**
+	* Prepare the scene view extension with relevant camera and render mode info.
+	*/
+	void PrepareRender(TSoftObjectPtr<UCineCameraComponent> InCineCameraComponent, bool bInFollowSceneCaptureRenderPath)
 	{
-		CineCaptureSVE->OnViewSetupFinished();
+		CineCameraComponent = InCineCameraComponent;
+		bFollowSceneCaptureRenderPath = bInFollowSceneCaptureRenderPath;
+	};
+
+	/**
+	* Delta time is needed for the purposes of camera smoothing.
+	*/
+	void SetFrameDeltaTime(float InDeltaTime)
+	{
+		DeltaTime = InDeltaTime;
 	}
 private:
-	TSharedPtr<FCineCameraCaptureSceneViewExtension, ESPMode::ThreadSafe> CineCaptureSVE;
+	/**
+	* Translates to View.bIsSceneCapture.
+	*/
+	bool bFollowSceneCaptureRenderPath = true;
+
+	/**
+	* A transient property that is used to deliver settings from CineCameraComponent to views that are related to cine capture.
+	*/
+	TSoftObjectPtr<UCineCameraComponent> CineCameraComponent = nullptr;
+
+	/**
+	* Delta time between frames. Used for camera smoothing.
+	*/
+	float DeltaTime = 0.0f;
 };
 
-
 UCineCaptureComponent2D::UCineCaptureComponent2D(const FObjectInitializer& ObjectInitializer)
 	: Super(ObjectInitializer), RenderTargetHighestDimension(1280), bFollowSceneCaptureRenderPath(true)
 {
@@ -52,19 +96,13 @@ void UCineCaptureComponent2D::UpdateSceneCaptureContents(FSceneInterface* Scene)
 	SetRelativeTransform(FTransform::Identity);
 	
 	PostProcessSettings = CineCameraComponent->PostProcessSettings;
-
-#if WITH_EDITORONLY_DATA
-	FOVAngle = CineCameraComponent->CurrentHorizontalFOV;
-#endif
-
+	FOVAngle = CineCameraComponent->FieldOfView;
 	bOverride_CustomNearClippingPlane = CineCameraComponent->bOverride_CustomNearClippingPlane;
 	CustomNearClippingPlane = CineCameraComponent->CustomNearClippingPlane;
 
-	{
-		/* Since SVE setup is happening on game thread, this guarantees that all SVEs unrelated to CineCapture will be excluded from modification. */
-		FScopedCineCaptureSVESetup ScopedCineCaptureSVESetup(CineCaptureSVE, CineCameraComponent, bFollowSceneCaptureRenderPath);
-		Scene->UpdateSceneCaptureContents(this);
-	}
+	CineCaptureSVE->PrepareRender(CineCameraComponent, bFollowSceneCaptureRenderPath);
+
+	Scene->UpdateSceneCaptureContents(this);
 }
 
 void UCineCaptureComponent2D::CheckResizeRenderTarget()
@@ -99,22 +137,28 @@ void UCineCaptureComponent2D::OnRegister()
 	// Remove mesh created by Scene Capture Component.
 	ProxyMeshComponent->DestroyComponent();
 	ProxyMeshComponent = nullptr;
+#endif
+
 	if (!CineCaptureSVE.IsValid())
 	{
 		// This will create a new scene extension per each capture.
-		CineCaptureSVE = FSceneViewExtensions::NewExtension<FCineCameraCaptureSceneViewExtension>();
+		CineCaptureSVE = MakeShared<FCineCameraCaptureSceneViewExtension>();
+		SceneViewExtensions.Add(CineCaptureSVE);
 	}
-#endif
+
 	ValidateCineCameraComponent();
 }
 
 void UCineCaptureComponent2D::OnUnregister()
 {
 	Super::OnUnregister();
-	if (!CineCaptureSVE.IsValid())
+
+	if (CineCaptureSVE.IsValid())
 	{
+		SceneViewExtensions.Remove(CineCaptureSVE);
 		CineCaptureSVE.Reset();
 	}
+
 	CineCameraComponent = nullptr;
 }
 
@@ -165,37 +209,5 @@ void UCineCaptureComponent2D::ValidateCineCameraComponent()
 	}
 }
 
-FCineCameraCaptureSceneViewExtension::FCineCameraCaptureSceneViewExtension(const FAutoRegister& AutoRegister)
-	: FSceneViewExtensionBase(AutoRegister), bSetupViewInProcess(false)
-{
-
-}
-
-void FCineCameraCaptureSceneViewExtension::SetupView(FSceneViewFamily& InViewFamily, FSceneView& InView)
-{
-	if (!CineCameraComponent.IsValid())
-	{
-		return;
-	}
-
-	if (bSetupViewInProcess)
-	{
-		FMinimalViewInfo DesiredView;
-
-		// Immitate the behaviour of viewports.
-		CineCameraComponent->GetCameraView(DeltaTime, DesiredView);
-		FTransform Transform = CineCameraComponent->GetComponentToWorld();
-		FVector ViewLocation = Transform.GetTranslation();
-		InView.StartFinalPostprocessSettings(ViewLocation);
-
-		InView.OverridePostProcessSettings(DesiredView.PostProcessSettings, 1.0);
-		FSceneViewInitOptions ViewInitOptions;
-		InView.EndFinalPostprocessSettings(ViewInitOptions);
-
-		// Required for certain effects (lighting) to match that of Cine Camera.
-		InView.bIsSceneCapture = bFollowSceneCaptureRenderPath;
-	}
-}
-
 
 #undef LOCTEXT_NAMESPACE
diff --git a/Engine/Plugins/Experimental/CineCameraSceneCapture/Source/CineCameraSceneCapture/Public/CineCameraSceneCaptureComponent.h b/Engine/Plugins/Experimental/CineCameraSceneCapture/Source/CineCameraSceneCapture/Public/CineCameraSceneCaptureComponent.h
--- a/Engine/Plugins/Experimental/CineCameraSceneCapture/Source/CineCameraSceneCapture/Public/CineCameraSceneCaptureComponent.h
+++ b/Engine/Plugins/Experimental/CineCameraSceneCapture/Source/CineCameraSceneCapture/Public/CineCameraSceneCaptureComponent.h
@@ -3,68 +3,13 @@
 #pragma once
 
 #include "Components/SceneCaptureComponent2D.h"
-#include "SceneViewExtension.h"
 #include "CineCameraSceneCaptureComponent.generated.h"
 
 DECLARE_LOG_CATEGORY_EXTERN(LogCineCapture, Log, All);
 
+class FCineCameraCaptureSceneViewExtension;
 class UCineCameraComponent;
 
-class FCineCameraCaptureSceneViewExtension : public FSceneViewExtensionBase
-{
-public:
-	FCineCameraCaptureSceneViewExtension(const FAutoRegister& AutoRegister);
-
-	//~ Begin FSceneViewExtensionBase Interface
-	virtual void SetupViewFamily(FSceneViewFamily& InViewFamily) override {};
-	virtual void SetupView(FSceneViewFamily& InViewFamily, FSceneView& InView) override;
-	virtual void BeginRenderViewFamily(FSceneViewFamily& InViewFamily) override {};
-	//~ End FSceneViewExtensionBase Interface
-public:
-	/** 
-	* Flags all views that will be setup until (bSetupViewInProcess == false) as related to Cine Capture.
-	*/
-	void OnViewSetupStarted(TSoftObjectPtr<UCineCameraComponent> InCineCameraComponent, bool bInFollowSceneCaptureRenderPath)
-	{ 
-		bSetupViewInProcess = true; 
-		CineCameraComponent = InCineCameraComponent;
-		bFollowSceneCaptureRenderPath = bInFollowSceneCaptureRenderPath;
-	};
-
-	/**
-	* Flags all following view setups as not related to Cine Capture.
-	*/
-	void OnViewSetupFinished() { bSetupViewInProcess = false; };
-
-	/**
-	* Delta time is needed for the purposes of camera smoothing.
-	*/
-	void SetFrameDeltaTime(float InDeltaTime)
-	{
-		DeltaTime = InDeltaTime;
-	}
-private:
-	/** 
-	* Translates to View.bIsSceneCapture.
-	*/
-	bool bFollowSceneCaptureRenderPath;
-
-	/** 
-	* When this property is true, it means that the views that are being setup currently are related to Cine Capture.
-	*/
-	bool bSetupViewInProcess;
-
-	/**
-	* A transient property that is used to deliver settings from CineCameraComponent to views that are related to cine capture.
-	*/
-	TSoftObjectPtr<UCineCameraComponent> CineCameraComponent;
-
-	/** 
-	* Delta time between frames. Used for camera smoothing.
-	*/
-	float DeltaTime;
-};
-
 /**
 * Cine Capture Component extends Scene Capture to allow users to render Cine Camera Component into a render target. 
 * Cine Capture has a few modifiable properties, but most of the properties are controlled by Cine Camera Component.
@@ -94,8 +39,6 @@ public:
 	UPROPERTY()
 	TSoftObjectPtr<UCineCameraComponent> CineCameraComponent;
 
-public:
-	void CheckResizeRenderTarget();
 public:
 	virtual void UpdateSceneCaptureContents(FSceneInterface* Scene) override;
 	virtual void OnRegister() override;
@@ -105,7 +48,8 @@ public:
 	virtual void OnAttachmentChanged() override; 
 	virtual void TickComponent(float DeltaTime, enum ELevelTick TickType, FActorComponentTickFunction* ThisTickFunction) override;
 
-public:
+private:
+	void CheckResizeRenderTarget();
 	void ValidateCineCameraComponent();
 
 private:
