Cine Camera Capture: Added OCIO transform support.

diff --git a/Engine/Plugins/Experimental/CineCameraSceneCapture/Source/CineCameraSceneCapture/CineCameraSceneCapture.Build.cs b/Engine/Plugins/Experimental/CineCameraSceneCapture/Source/CineCameraSceneCapture/CineCameraSceneCapture.Build.cs
--- a/Engine/Plugins/Experimental/CineCameraSceneCapture/Source/CineCameraSceneCapture/CineCameraSceneCapture.Build.cs
+++ b/Engine/Plugins/Experimental/CineCameraSceneCapture/Source/CineCameraSceneCapture/CineCameraSceneCapture.Build.cs
@@ -15,6 +15,15 @@ public class CineCameraSceneCapture : ModuleRules
 				"Slate",
 				"SlateCore",
 				"CinematicCamera",
+				"RenderCore",
+				"Renderer",
+			}
+		);
+
+		PublicDependencyModuleNames.AddRange(
+			new string[]
+			{
+				"OpenColorIO",
 			}
 		);
 	}
diff --git a/Engine/Plugins/Experimental/CineCameraSceneCapture/Source/CineCameraSceneCapture/Private/CineCameraSceneCaptureComponent.cpp b/Engine/Plugins/Experimental/CineCameraSceneCapture/Source/CineCameraSceneCapture/Private/CineCameraSceneCaptureComponent.cpp
--- a/Engine/Plugins/Experimental/CineCameraSceneCapture/Source/CineCameraSceneCapture/Private/CineCameraSceneCaptureComponent.cpp
+++ b/Engine/Plugins/Experimental/CineCameraSceneCapture/Source/CineCameraSceneCapture/Private/CineCameraSceneCaptureComponent.cpp
@@ -7,6 +7,10 @@
 #include "Runtime/Engine/Public/SceneView.h"
 #include "Engine/TextureRenderTarget2D.h"
 #include "Framework/Notifications/NotificationManager.h"
+#include "OpenColorIODisplayExtension.h"
+#include "OpenColorIORendering.h"
+#include "PostProcess/PostProcessMaterialInputs.h"
+#include "ScreenPass.h"
 #include "Widgets/Notifications/SNotificationList.h"
 
 #define LOCTEXT_NAMESPACE "FCineCameraSceneCaptureModule"
@@ -42,7 +46,29 @@ public:
 
 		// Required for certain effects (lighting) to match that of Cine Camera.
 		InView.bIsSceneCapture = bFollowSceneCaptureRenderPath;
+
+		bIsOcioEnabledRenderThread = DisplayConfiguration.bIsEnabled;
+
+		// Setup OCIO.
+		if (DisplayConfiguration.bIsEnabled)
+		{
+			FOpenColorIORenderPassResources PassResources = FOpenColorIORendering::GetRenderPassResources(DisplayConfiguration.ColorConfiguration, InViewFamily.GetFeatureLevel());
+
+			if (PassResources.IsValid())
+			{
+				FOpenColorIORendering::PrepareView(InViewFamily, InView);
+			}
+
+			ENQUEUE_RENDER_COMMAND(ProcessColorSpaceTransform)(
+				[this, ResourcesRenderThread = MoveTemp(PassResources), bInIsOcioEnabledRenderThread = DisplayConfiguration.bIsEnabled](FRHICommandListImmediate& RHICmdList)
+				{
+					//Caches render thread resource to be used when applying configuration in PostRenderViewFamily_RenderThread
+					CachedResourcesRenderThread = ResourcesRenderThread;
+				}
+			);
+		}
 	}
+
 	virtual void BeginRenderViewFamily(FSceneViewFamily& InViewFamily) override {};
 	//~ End FSceneViewExtensionBase Interface
 
@@ -62,6 +88,46 @@ public:
 	{
 		DeltaTime = InDeltaTime;
 	}
+
+	virtual void SubscribeToPostProcessingPass(EPostProcessingPass PassId, FAfterPassCallbackDelegateArray& InOutPassCallbacks, bool bIsPassEnabled)
+	{
+		if (!bIsOcioEnabledRenderThread)
+		{
+			return;
+		}
+
+		if (PassId == EPostProcessingPass::Tonemap)
+		{
+			InOutPassCallbacks.Add(FAfterPassCallbackDelegate::CreateRaw(this, &FCineCameraCaptureSceneViewExtension::PostProcessPassAfterTonemap_RenderThread));
+		}
+	}
+
+	FScreenPassTexture PostProcessPassAfterTonemap_RenderThread(FRDGBuilder& GraphBuilder, const FSceneView& View, const FPostProcessMaterialInputs& InOutInputs)
+	{
+		const FScreenPassTexture& SceneColor = InOutInputs.GetInput(EPostProcessMaterialInput::SceneColor);
+		check(SceneColor.IsValid());
+
+		FScreenPassRenderTarget Output = InOutInputs.OverrideOutput;
+
+		// If the override output is provided, it means that this is the last pass in post processing.
+		if (!Output.IsValid())
+		{
+			Output = FScreenPassRenderTarget::CreateFromInput(GraphBuilder, SceneColor, View.GetOverwriteLoadAction(), TEXT("OCIORenderTarget"));
+		}
+
+		FOpenColorIORendering::AddPass_RenderThread(
+			GraphBuilder,
+			View,
+			SceneColor,
+			Output,
+			CachedResourcesRenderThread
+		);
+
+		return MoveTemp(Output);
+	}
+
+	void SetDisplayConfiguration(const FOpenColorIODisplayConfiguration& InDisplayConfiguration) { DisplayConfiguration = InDisplayConfiguration; };
+
 private:
 	/**
 	* Translates to View.bIsSceneCapture.
@@ -77,13 +143,24 @@ private:
 	* Delta time between frames. Used for camera smoothing.
 	*/
 	float DeltaTime = 0.0f;
+
+	/** Cached pass resources required to apply conversion for render thread */
+	FOpenColorIORenderPassResources CachedResourcesRenderThread;
+
+	/** Configuration to apply during post render callback */
+	FOpenColorIODisplayConfiguration DisplayConfiguration;
+
+	/** Indicates if OCIO is enabled. */
+	std::atomic<bool> bIsOcioEnabledRenderThread;
 };
 
+
 UCineCaptureComponent2D::UCineCaptureComponent2D(const FObjectInitializer& ObjectInitializer)
 	: Super(ObjectInitializer), RenderTargetHighestDimension(1280), bFollowSceneCaptureRenderPath(true)
 {
 	CaptureSource = ESceneCaptureSource::SCS_FinalToneCurveHDR;
 	bAlwaysPersistRenderingState = true;
+	OCIOConfiguration.bIsEnabled = false;
 }
 
 void UCineCaptureComponent2D::UpdateSceneCaptureContents(FSceneInterface* Scene)
@@ -146,6 +223,8 @@ void UCineCaptureComponent2D::OnRegister()
 		SceneViewExtensions.Add(CineCaptureSVE);
 	}
 
+	CineCaptureSVE->SetDisplayConfiguration(OCIOConfiguration);
+
 	ValidateCineCameraComponent();
 }
 
diff --git a/Engine/Plugins/Experimental/CineCameraSceneCapture/Source/CineCameraSceneCapture/Public/CineCameraSceneCaptureComponent.h b/Engine/Plugins/Experimental/CineCameraSceneCapture/Source/CineCameraSceneCapture/Public/CineCameraSceneCaptureComponent.h
--- a/Engine/Plugins/Experimental/CineCameraSceneCapture/Source/CineCameraSceneCapture/Public/CineCameraSceneCaptureComponent.h
+++ b/Engine/Plugins/Experimental/CineCameraSceneCapture/Source/CineCameraSceneCapture/Public/CineCameraSceneCaptureComponent.h
@@ -2,6 +2,7 @@
 
 #pragma once
 
+#include "OpenColorIOColorSpace.h"
 #include "Components/SceneCaptureComponent2D.h"
 #include "CineCameraSceneCaptureComponent.generated.h"
 
@@ -36,6 +37,13 @@ public:
 	UPROPERTY(EditAnywhere, BlueprintReadWrite, Category = "Capture Settings", AdvancedDisplay)
 	bool bFollowSceneCaptureRenderPath;
 
+	/** 
+	* Open Color IO Configuration. Used to apply color transforms.
+	*/
+	UPROPERTY(EditAnywhere, BlueprintReadWrite, Category = "Color", meta = (DisplayName = "OpenColorIO Display Configuration"))
+	FOpenColorIODisplayConfiguration OCIOConfiguration;
+
+	/** Used to gather relevant properties for View rendering. */
 	UPROPERTY()
 	TSoftObjectPtr<UCineCameraComponent> CineCameraComponent;
 
