Cine Camera Capture: Minor cleanup.

diff --git a/Engine/Plugins/Experimental/CineCameraSceneCapture/Source/CineCameraSceneCapture/Private/CineCameraSceneCaptureComponent.cpp b/Engine/Plugins/Experimental/CineCameraSceneCapture/Source/CineCameraSceneCapture/Private/CineCameraSceneCaptureComponent.cpp
--- a/Engine/Plugins/Experimental/CineCameraSceneCapture/Source/CineCameraSceneCapture/Private/CineCameraSceneCaptureComponent.cpp
+++ b/Engine/Plugins/Experimental/CineCameraSceneCapture/Source/CineCameraSceneCapture/Private/CineCameraSceneCaptureComponent.cpp
@@ -60,7 +60,7 @@ public:
 			}
 
 			ENQUEUE_RENDER_COMMAND(ProcessColorSpaceTransform)(
-				[this, ResourcesRenderThread = MoveTemp(PassResources), bInIsOcioEnabledRenderThread = DisplayConfiguration.bIsEnabled](FRHICommandListImmediate& RHICmdList)
+				[this, ResourcesRenderThread = MoveTemp(PassResources)](FRHICommandListImmediate& RHICmdList)
 				{
 					//Caches render thread resource to be used when applying configuration in PostRenderViewFamily_RenderThread
 					CachedResourcesRenderThread = ResourcesRenderThread;
