CineCapture: Ensure OpenColorIO is listed as a plugin dependency.

#rnx

diff --git a/Engine/Plugins/Experimental/CineCameraSceneCapture/CineCameraSceneCapture.uplugin b/Engine/Plugins/Experimental/CineCameraSceneCapture/CineCameraSceneCapture.uplugin
--- a/Engine/Plugins/Experimental/CineCameraSceneCapture/CineCameraSceneCapture.uplugin
+++ b/Engine/Plugins/Experimental/CineCameraSceneCapture/CineCameraSceneCapture.uplugin
@@ -21,5 +21,11 @@
 			"Type": "Runtime",
 			"LoadingPhase": "Default"
 		}
+	],
+	"Plugins": [
+		{
+			"Name": "OpenColorIO",
+			"Enabled": true
+		}
 	]
 }
\ No newline at end of file
