Cine Camera Capture Component: Fixed an issue caused by Camera mesh not being added for some reason.

diff --git a/Engine/Plugins/Experimental/CineCameraSceneCapture/Source/CineCameraSceneCapture/Private/CineCameraSceneCaptureComponent.cpp b/Engine/Plugins/Experimental/CineCameraSceneCapture/Source/CineCameraSceneCapture/Private/CineCameraSceneCaptureComponent.cpp
--- a/Engine/Plugins/Experimental/CineCameraSceneCapture/Source/CineCameraSceneCapture/Private/CineCameraSceneCaptureComponent.cpp
+++ b/Engine/Plugins/Experimental/CineCameraSceneCapture/Source/CineCameraSceneCapture/Private/CineCameraSceneCaptureComponent.cpp
@@ -211,8 +211,11 @@ void UCineCaptureComponent2D::OnRegister()
 
 #if WITH_EDITORONLY_DATA
 	// Remove mesh created by Scene Capture Component.
-	ProxyMeshComponent->DestroyComponent();
-	ProxyMeshComponent = nullptr;
+	if (ProxyMeshComponent)
+	{
+		ProxyMeshComponent->DestroyComponent();
+		ProxyMeshComponent = nullptr;
+	}
 #endif
 
 	if (!CineCaptureSVE.IsValid())
