Cine Camera Capture Component: Exposed component to external plugins.

diff --git a/Engine/Plugins/Experimental/CineCameraSceneCapture/Source/CineCameraSceneCapture/Public/CineCameraSceneCaptureComponent.h b/Engine/Plugins/Experimental/CineCameraSceneCapture/Source/CineCameraSceneCapture/Public/CineCameraSceneCaptureComponent.h
--- a/Engine/Plugins/Experimental/CineCameraSceneCapture/Source/CineCameraSceneCapture/Public/CineCameraSceneCaptureComponent.h
+++ b/Engine/Plugins/Experimental/CineCameraSceneCapture/Source/CineCameraSceneCapture/Public/CineCameraSceneCaptureComponent.h
@@ -18,7 +18,7 @@ class UCineCameraComponent;
 * 
 */
 UCLASS(hidecategories = (Transform, Collision, Object, Physics, SceneComponent, PostProcessVolume, Projection, Rendering, PlanarReflection), ClassGroup = Rendering, editinlinenew, meta = (BlueprintSpawnableComponent))
-class UCineCaptureComponent2D : public USceneCaptureComponent2D
+class CINECAMERASCENECAPTURE_API UCineCaptureComponent2D : public USceneCaptureComponent2D
 {
 	GENERATED_UCLASS_BODY()
 public:
