OCIO: Fix scene view extension potential use-after-free in enqueued render commands.

diff --git a/Engine/Plugins/Compositing/OpenColorIO/Source/OpenColorIO/Private/OpenColorIODisplayExtension.cpp b/Engine/Plugins/Compositing/OpenColorIO/Source/OpenColorIO/Private/OpenColorIODisplayExtension.cpp
--- a/Engine/Plugins/Compositing/OpenColorIO/Source/OpenColorIO/Private/OpenColorIODisplayExtension.cpp
+++ b/Engine/Plugins/Compositing/OpenColorIO/Source/OpenColorIO/Private/OpenColorIODisplayExtension.cpp
@@ -54,10 +54,14 @@ void FOpenColorIODisplayExtension::SetupView(FSceneViewFamily& InViewFamily, FSc
 		FOpenColorIORendering::PrepareView(InViewFamily, InView);
 	}
 	ENQUEUE_RENDER_COMMAND(ProcessColorSpaceTransform)(
-		[this, ResourcesRenderThread = MoveTemp(PassResources)](FRHICommandListImmediate& RHICmdList)
+		[WeakThis = AsWeak(), ResourcesRenderThread = MoveTemp(PassResources)](FRHICommandListImmediate& RHICmdList)
 		{
-			//Caches render thread resource to be used when applying configuration in PostRenderViewFamily_RenderThread
-			CachedResourcesRenderThread = ResourcesRenderThread;
+			TSharedPtr<FOpenColorIODisplayExtension> This = StaticCastWeakPtr<FOpenColorIODisplayExtension>(WeakThis).Pin();
+			if (This.IsValid())
+			{
+				//Caches render thread resource to be used when applying configuration in PostRenderViewFamily_RenderThread
+				This->CachedResourcesRenderThread = ResourcesRenderThread;
+			}
 		}
 	);
 }
diff --git a/Engine/Plugins/Experimental/CineCameraSceneCapture/Source/CineCameraSceneCapture/Private/CineCameraSceneCaptureComponent.cpp b/Engine/Plugins/Experimental/CineCameraSceneCapture/Source/CineCameraSceneCapture/Private/CineCameraSceneCaptureComponent.cpp
--- a/Engine/Plugins/Experimental/CineCameraSceneCapture/Source/CineCameraSceneCapture/Private/CineCameraSceneCaptureComponent.cpp
+++ b/Engine/Plugins/Experimental/CineCameraSceneCapture/Source/CineCameraSceneCapture/Private/CineCameraSceneCaptureComponent.cpp
@@ -20,7 +20,7 @@ DEFINE_LOG_CATEGORY(LogCineCapture);
 #define CINE_CAMERA_INVALID_PARENT_WARNING "Cine Capture requires to be parented to Cine Camera Component. Cine Capture {0} on Actor \"{1}\" will be disabled until it is parented to Cine Camera Actor."
 
 // This extension is only registered onto the scene capture 2d component, and therefore runs locally.
-class FCineCameraCaptureSceneViewExtension : public ISceneViewExtension
+class FCineCameraCaptureSceneViewExtension : public ISceneViewExtension, public TSharedFromThis<FCineCameraCaptureSceneViewExtension, ESPMode::ThreadSafe>
 {
 public:
 	//~ Begin FSceneViewExtensionBase Interface
@@ -60,10 +60,14 @@ public:
 			}
 
 			ENQUEUE_RENDER_COMMAND(ProcessColorSpaceTransform)(
-				[this, ResourcesRenderThread = MoveTemp(PassResources)](FRHICommandListImmediate& RHICmdList)
+				[WeakThis = AsWeak(), ResourcesRenderThread = MoveTemp(PassResources)](FRHICommandListImmediate& RHICmdList)
 				{
-					//Caches render thread resource to be used when applying configuration in PostRenderViewFamily_RenderThread
-					CachedResourcesRenderThread = ResourcesRenderThread;
+					TSharedPtr<FCineCameraCaptureSceneViewExtension> This = WeakThis.Pin();
+					if (This.IsValid())
+					{
+						//Caches render thread resource to be used when applying configuration in PostRenderViewFamily_RenderThread
+						This->CachedResourcesRenderThread = ResourcesRenderThread;
+					}
 				}
 			);
 		}
