CineCameraCapture: Added option to discard world post-process volume blendables.

#rnx

diff --git a/Engine/Plugins/Experimental/CineCameraSceneCapture/Source/CineCameraSceneCapture/Private/CineCameraSceneCaptureComponent.cpp b/Engine/Plugins/Experimental/CineCameraSceneCapture/Source/CineCameraSceneCapture/Private/CineCameraSceneCaptureComponent.cpp
--- a/Engine/Plugins/Experimental/CineCameraSceneCapture/Source/CineCameraSceneCapture/Private/CineCameraSceneCaptureComponent.cpp
+++ b/Engine/Plugins/Experimental/CineCameraSceneCapture/Source/CineCameraSceneCapture/Private/CineCameraSceneCaptureComponent.cpp
@@ -40,6 +40,8 @@ public:
 		FVector ViewLocation = Transform.GetTranslation();
 		InView.StartFinalPostprocessSettings(ViewLocation);
 
+		DesiredView.PostProcessSettings.AddBlendable(BlendableObject, 1.0f);
+
 		InView.OverridePostProcessSettings(DesiredView.PostProcessSettings, 1.0);
 		FSceneViewInitOptions ViewInitOptions;
 		InView.EndFinalPostprocessSettings(ViewInitOptions);
@@ -79,10 +81,11 @@ public:
 	/**
 	* Prepare the scene view extension with relevant camera and render mode info.
 	*/
-	void PrepareRender(TSoftObjectPtr<UCineCameraComponent> InCineCameraComponent, bool bInFollowSceneCaptureRenderPath)
+	void PrepareRender(TSoftObjectPtr<UCineCameraComponent> InCineCameraComponent, bool bInFollowSceneCaptureRenderPath, TScriptInterface<IBlendableInterface> InBlendableObject)
 	{
 		CineCameraComponent = InCineCameraComponent;
 		bFollowSceneCaptureRenderPath = bInFollowSceneCaptureRenderPath;
+		BlendableObject = InBlendableObject;
 	};
 
 	/**
@@ -156,14 +159,29 @@ private:
 
 	/** Indicates if OCIO is enabled. */
 	std::atomic<bool> bIsOcioEnabledRenderThread;
+
+	/** Optional blendable object. */
+	TScriptInterface<IBlendableInterface> BlendableObject = nullptr;
 };
 
+UCineCaptureErasureBlendable::UCineCaptureErasureBlendable(const FObjectInitializer& ObjectInitializer)
+	: Super(ObjectInitializer)
+{}
+
+void UCineCaptureErasureBlendable::OverrideBlendableSettings(FSceneView& View, float Weight) const
+{
+	// Clear any blendables that could have been set by post process volumes.
+	View.FinalPostProcessSettings.BlendableManager = FBlendableManager();
+}
+
 
 UCineCaptureComponent2D::UCineCaptureComponent2D(const FObjectInitializer& ObjectInitializer)
-	: Super(ObjectInitializer), RenderTargetHighestDimension(1280), bFollowSceneCaptureRenderPath(true)
+	: Super(ObjectInitializer), RenderTargetHighestDimension(1280), bFollowSceneCaptureRenderPath(true), bIgnoreWorldPostProcessMaterials(false)
 {
 	CaptureSource = ESceneCaptureSource::SCS_FinalToneCurveHDR;
 	bAlwaysPersistRenderingState = true;
+
+	ErasureBlendable = ObjectInitializer.CreateDefaultSubobject<UCineCaptureErasureBlendable>(this, TEXT("CineCaptureErasureBlendable"));
 }
 
 void UCineCaptureComponent2D::UpdateSceneCaptureContents(FSceneInterface* Scene)
@@ -180,7 +198,7 @@ void UCineCaptureComponent2D::UpdateSceneCaptureContents(FSceneInterface* Scene)
 	bOverride_CustomNearClippingPlane = CineCameraComponent->bOverride_CustomNearClippingPlane;
 	CustomNearClippingPlane = CineCameraComponent->CustomNearClippingPlane;
 
-	CineCaptureSVE->PrepareRender(CineCameraComponent, bFollowSceneCaptureRenderPath);
+	CineCaptureSVE->PrepareRender(CineCameraComponent, bFollowSceneCaptureRenderPath, bIgnoreWorldPostProcessMaterials ? ErasureBlendable : nullptr);
 
 	Scene->UpdateSceneCaptureContents(this);
 }
diff --git a/Engine/Plugins/Experimental/CineCameraSceneCapture/Source/CineCameraSceneCapture/Public/CineCameraSceneCaptureComponent.h b/Engine/Plugins/Experimental/CineCameraSceneCapture/Source/CineCameraSceneCapture/Public/CineCameraSceneCaptureComponent.h
--- a/Engine/Plugins/Experimental/CineCameraSceneCapture/Source/CineCameraSceneCapture/Public/CineCameraSceneCaptureComponent.h
+++ b/Engine/Plugins/Experimental/CineCameraSceneCapture/Source/CineCameraSceneCapture/Public/CineCameraSceneCaptureComponent.h
@@ -4,6 +4,7 @@
 
 #include "OpenColorIOColorSpace.h"
 #include "Components/SceneCaptureComponent2D.h"
+#include "Engine/BlendableInterface.h"
 #include "CineCameraSceneCaptureComponent.generated.h"
 
 DECLARE_LOG_CATEGORY_EXTERN(LogCineCapture, Log, All);
@@ -11,6 +12,22 @@ DECLARE_LOG_CATEGORY_EXTERN(LogCineCapture, Log, All);
 class FCineCameraCaptureSceneViewExtension;
 class UCineCameraComponent;
 
+/**
+* Custom blendable used to optionally clear world blendables (post-process materials) in the cine camera capture.
+*/
+UCLASS(NotBlueprintType)
+class UCineCaptureErasureBlendable
+	: public UObject
+	, public IBlendableInterface
+{
+	GENERATED_UCLASS_BODY()
+
+public:
+	// Begins IBlendableInterface
+	void OverrideBlendableSettings(class FSceneView& View, float Weight) const override;
+	// Ends IBlendableInterface
+};
+
 /**
 * Cine Capture Component extends Scene Capture to allow users to render Cine Camera Component into a render target. 
 * Cine Capture has a few modifiable properties, but most of the properties are controlled by Cine Camera Component.
@@ -43,6 +60,12 @@ public:
 	UPROPERTY(EditAnywhere, BlueprintReadWrite, Category = "Color", meta = (DisplayName = "OpenColorIO Display Configuration"))
 	FOpenColorIODisplayConfiguration OCIOConfiguration;
 
+	/**
+	* Option to disable world blendables (post-process materials) only in the cine camera capture.
+	*/
+	UPROPERTY(EditAnywhere, BlueprintReadWrite, Category = "Capture Settings")
+	bool bIgnoreWorldPostProcessMaterials;
+
 	/** Used to gather relevant properties for View rendering. */
 	UPROPERTY()
 	TSoftObjectPtr<UCineCameraComponent> CineCameraComponent;
@@ -63,4 +86,8 @@ private:
 private:
 	/** This scene view extension is used to get ahold of views during the setup process. */
 	TSharedPtr<FCineCameraCaptureSceneViewExtension, ESPMode::ThreadSafe> CineCaptureSVE;
+
+	UPROPERTY()
+	/** Custom blendable used to optionally clear world blendables (post-process materials) in the cine camera capture. */
+	TObjectPtr<UCineCaptureErasureBlendable> ErasureBlendable;
 };
