[Backout] - CL33897916
Original CL Desc
-----------------------------------------------------------------
CineCameraCapture: Added option to discard world post-process volume blendables.

#rnx

diff --git a/Engine/Plugins/Experimental/CineCameraSceneCapture/Source/CineCameraSceneCapture/Private/CineCameraSceneCaptureComponent.cpp b/Engine/Plugins/Experimental/CineCameraSceneCapture/Source/CineCameraSceneCapture/Private/CineCameraSceneCaptureComponent.cpp
--- a/Engine/Plugins/Experimental/CineCameraSceneCapture/Source/CineCameraSceneCapture/Private/CineCameraSceneCaptureComponent.cpp
+++ b/Engine/Plugins/Experimental/CineCameraSceneCapture/Source/CineCameraSceneCapture/Private/CineCameraSceneCaptureComponent.cpp
@@ -40,8 +40,6 @@ public:
 		FVector ViewLocation = Transform.GetTranslation();
 		InView.StartFinalPostprocessSettings(ViewLocation);
 
-		DesiredView.PostProcessSettings.AddBlendable(BlendableObject, 1.0f);
-
 		InView.OverridePostProcessSettings(DesiredView.PostProcessSettings, 1.0);
 		FSceneViewInitOptions ViewInitOptions;
 		InView.EndFinalPostprocessSettings(ViewInitOptions);
@@ -81,11 +79,10 @@ public:
 	/**
 	* Prepare the scene view extension with relevant camera and render mode info.
 	*/
-	void PrepareRender(TSoftObjectPtr<UCineCameraComponent> InCineCameraComponent, bool bInFollowSceneCaptureRenderPath, TScriptInterface<IBlendableInterface> InBlendableObject)
+	void PrepareRender(TSoftObjectPtr<UCineCameraComponent> InCineCameraComponent, bool bInFollowSceneCaptureRenderPath)
 	{
 		CineCameraComponent = InCineCameraComponent;
 		bFollowSceneCaptureRenderPath = bInFollowSceneCaptureRenderPath;
-		BlendableObject = InBlendableObject;
 	};
 
 	/**
@@ -159,29 +156,14 @@ private:
 
 	/** Indicates if OCIO is enabled. */
 	std::atomic<bool> bIsOcioEnabledRenderThread;
-
-	/** Optional blendable object. */
-	TScriptInterface<IBlendableInterface> BlendableObject = nullptr;
 };
 
-UCineCaptureErasureBlendable::UCineCaptureErasureBlendable(const FObjectInitializer& ObjectInitializer)
-	: Super(ObjectInitializer)
-{}
-
-void UCineCaptureErasureBlendable::OverrideBlendableSettings(FSceneView& View, float Weight) const
-{
-	// Clear any blendables that could have been set by post process volumes.
-	View.FinalPostProcessSettings.BlendableManager = FBlendableManager();
-}
-
 
 UCineCaptureComponent2D::UCineCaptureComponent2D(const FObjectInitializer& ObjectInitializer)
-	: Super(ObjectInitializer), RenderTargetHighestDimension(1280), bFollowSceneCaptureRenderPath(true), bIgnoreWorldPostProcessMaterials(false)
+	: Super(ObjectInitializer), RenderTargetHighestDimension(1280), bFollowSceneCaptureRenderPath(true)
 {
 	CaptureSource = ESceneCaptureSource::SCS_FinalToneCurveHDR;
 	bAlwaysPersistRenderingState = true;
-
-	ErasureBlendable = ObjectInitializer.CreateDefaultSubobject<UCineCaptureErasureBlendable>(this, TEXT("CineCaptureErasureBlendable"));
 }
 
 void UCineCaptureComponent2D::UpdateSceneCaptureContents(FSceneInterface* Scene)
@@ -198,7 +180,7 @@ void UCineCaptureComponent2D::UpdateSceneCaptureContents(FSceneInterface* Scene)
 	bOverride_CustomNearClippingPlane = CineCameraComponent->bOverride_CustomNearClippingPlane;
 	CustomNearClippingPlane = CineCameraComponent->CustomNearClippingPlane;
 
-	CineCaptureSVE->PrepareRender(CineCameraComponent, bFollowSceneCaptureRenderPath, bIgnoreWorldPostProcessMaterials ? ErasureBlendable : nullptr);
+	CineCaptureSVE->PrepareRender(CineCameraComponent, bFollowSceneCaptureRenderPath);
 
 	Scene->UpdateSceneCaptureContents(this);
 }
diff --git a/Engine/Plugins/Experimental/CineCameraSceneCapture/Source/CineCameraSceneCapture/Public/CineCameraSceneCaptureComponent.h b/Engine/Plugins/Experimental/CineCameraSceneCapture/Source/CineCameraSceneCapture/Public/CineCameraSceneCaptureComponent.h
--- a/Engine/Plugins/Experimental/CineCameraSceneCapture/Source/CineCameraSceneCapture/Public/CineCameraSceneCaptureComponent.h
+++ b/Engine/Plugins/Experimental/CineCameraSceneCapture/Source/CineCameraSceneCapture/Public/CineCameraSceneCaptureComponent.h
@@ -4,7 +4,6 @@
 
 #include "OpenColorIOColorSpace.h"
 #include "Components/SceneCaptureComponent2D.h"
-#include "Engine/BlendableInterface.h"
 #include "CineCameraSceneCaptureComponent.generated.h"
 
 DECLARE_LOG_CATEGORY_EXTERN(LogCineCapture, Log, All);
@@ -12,22 +11,6 @@ DECLARE_LOG_CATEGORY_EXTERN(LogCineCapture, Log, All);
 class FCineCameraCaptureSceneViewExtension;
 class UCineCameraComponent;
 
-/**
-* Custom blendable used to optionally clear world blendables (post-process materials) in the cine camera capture.
-*/
-UCLASS(NotBlueprintType)
-class UCineCaptureErasureBlendable
-	: public UObject
-	, public IBlendableInterface
-{
-	GENERATED_UCLASS_BODY()
-
-public:
-	// Begins IBlendableInterface
-	void OverrideBlendableSettings(class FSceneView& View, float Weight) const override;
-	// Ends IBlendableInterface
-};
-
 /**
 * Cine Capture Component extends Scene Capture to allow users to render Cine Camera Component into a render target. 
 * Cine Capture has a few modifiable properties, but most of the properties are controlled by Cine Camera Component.
@@ -60,12 +43,6 @@ public:
 	UPROPERTY(EditAnywhere, BlueprintReadWrite, Category = "Color", meta = (DisplayName = "OpenColorIO Display Configuration"))
 	FOpenColorIODisplayConfiguration OCIOConfiguration;
 
-	/**
-	* Option to disable world blendables (post-process materials) only in the cine camera capture.
-	*/
-	UPROPERTY(EditAnywhere, BlueprintReadWrite, Category = "Capture Settings")
-	bool bIgnoreWorldPostProcessMaterials;
-
 	/** Used to gather relevant properties for View rendering. */
 	UPROPERTY()
 	TSoftObjectPtr<UCineCameraComponent> CineCameraComponent;
@@ -86,8 +63,4 @@ private:
 private:
 	/** This scene view extension is used to get ahold of views during the setup process. */
 	TSharedPtr<FCineCameraCaptureSceneViewExtension, ESPMode::ThreadSafe> CineCaptureSVE;
-
-	UPROPERTY()
-	/** Custom blendable used to optionally clear world blendables (post-process materials) in the cine camera capture. */
-	TObjectPtr<UCineCaptureErasureBlendable> ErasureBlendable;
 };
