CineCameraCapture: Enable HWRT by default on cine camera scene capture renders.

#rnx

diff --git a/Engine/Plugins/Experimental/CineCameraSceneCapture/Source/CineCameraSceneCapture/Private/CineCameraSceneCaptureComponent.cpp b/Engine/Plugins/Experimental/CineCameraSceneCapture/Source/CineCameraSceneCapture/Private/CineCameraSceneCaptureComponent.cpp
--- a/Engine/Plugins/Experimental/CineCameraSceneCapture/Source/CineCameraSceneCapture/Private/CineCameraSceneCaptureComponent.cpp
+++ b/Engine/Plugins/Experimental/CineCameraSceneCapture/Source/CineCameraSceneCapture/Private/CineCameraSceneCaptureComponent.cpp
@@ -164,6 +164,9 @@ UCineCaptureComponent2D::UCineCaptureComponent2D(const FObjectInitializer& Objec
 {
 	CaptureSource = ESceneCaptureSource::SCS_FinalToneCurveHDR;
 	bAlwaysPersistRenderingState = true;
+
+	// Enable HWRT by default on the scene capture render path to minimize artifacts
+	bUseRayTracingIfEnabled = true;
 }
 
 void UCineCaptureComponent2D::UpdateSceneCaptureContents(FSceneInterface* Scene)
