CineCameraCapture: Exposed user flag for the scene capture, to allow branching behavior in materials.

Additional cleanup of the scene view extension.

diff --git a/Engine/Plugins/Experimental/CineCameraSceneCapture/Source/CineCameraSceneCapture/Private/CineCameraSceneCaptureComponent.cpp b/Engine/Plugins/Experimental/CineCameraSceneCapture/Source/CineCameraSceneCapture/Private/CineCameraSceneCaptureComponent.cpp
--- a/Engine/Plugins/Experimental/CineCameraSceneCapture/Source/CineCameraSceneCapture/Private/CineCameraSceneCaptureComponent.cpp
+++ b/Engine/Plugins/Experimental/CineCameraSceneCapture/Source/CineCameraSceneCapture/Private/CineCameraSceneCaptureComponent.cpp
@@ -23,11 +23,22 @@ DEFINE_LOG_CATEGORY(LogCineCapture);
 class FCineCameraCaptureSceneViewExtension : public ISceneViewExtension, public TSharedFromThis<FCineCameraCaptureSceneViewExtension, ESPMode::ThreadSafe>
 {
 public:
+
+	/** Constructor */
+	FCineCameraCaptureSceneViewExtension(TWeakObjectPtr<UCineCaptureComponent2D> InCineCaptureComponent)
+		: CineCaptureComponentWeak(InCineCaptureComponent)
+		, DeltaTime(0.0f)
+		, CachedResourcesRenderThread()
+		, bIsOcioEnabledRenderThread(false)
+	{ }
+
 	//~ Begin FSceneViewExtensionBase Interface
 	virtual void SetupViewFamily(FSceneViewFamily& InViewFamily) override {};
 	virtual void SetupView(FSceneViewFamily& InViewFamily, FSceneView& InView) override
 	{
-		if (!CineCameraComponent.IsValid())
+		TStrongObjectPtr<UCineCaptureComponent2D> CineCaptureComponent = CineCaptureComponentWeak.Pin();
+
+		if (!CineCaptureComponent.IsValid())
 		{
 			return;
 		}
@@ -35,24 +46,28 @@ public:
 		FMinimalViewInfo DesiredView;
 
 		// Immitate the behaviour of viewports.
-		CineCameraComponent->GetCameraView(DeltaTime, DesiredView);
-		FTransform Transform = CineCameraComponent->GetComponentToWorld();
+		CineCaptureComponent->GetCameraView(DeltaTime, DesiredView);
+		FTransform Transform = CineCaptureComponent->GetComponentToWorld();
 		FVector ViewLocation = Transform.GetTranslation();
 		InView.StartFinalPostprocessSettings(ViewLocation);
-
+		if (CineCaptureComponent->bOverrideUserFlags)
+		{
+			DesiredView.PostProcessSettings.bOverride_UserFlags = CineCaptureComponent->bOverrideUserFlags;
+			DesiredView.PostProcessSettings.UserFlags = CineCaptureComponent->UserFlags;
+		}
 		InView.OverridePostProcessSettings(DesiredView.PostProcessSettings, 1.0);
 		FSceneViewInitOptions ViewInitOptions;
 		InView.EndFinalPostprocessSettings(ViewInitOptions);
 
 		// Required for certain effects (lighting) to match that of Cine Camera.
-		InView.bIsSceneCapture = bFollowSceneCaptureRenderPath;
+		InView.bIsSceneCapture = CineCaptureComponent->bFollowSceneCaptureRenderPath;
 
-		bIsOcioEnabledRenderThread = DisplayConfiguration.bIsEnabled;
+		bIsOcioEnabledRenderThread = CineCaptureComponent->OCIOConfiguration.bIsEnabled;
 
 		// Setup OCIO.
-		if (DisplayConfiguration.bIsEnabled)
+		if (bIsOcioEnabledRenderThread)
 		{
-			FOpenColorIORenderPassResources PassResources = FOpenColorIORendering::GetRenderPassResources(DisplayConfiguration.ColorConfiguration, InViewFamily.GetFeatureLevel());
+			FOpenColorIORenderPassResources PassResources = FOpenColorIORendering::GetRenderPassResources(CineCaptureComponent->OCIOConfiguration.ColorConfiguration, InViewFamily.GetFeatureLevel());
 
 			if (PassResources.IsValid())
 			{
@@ -76,16 +91,7 @@ public:
 	virtual void BeginRenderViewFamily(FSceneViewFamily& InViewFamily) override {};
 	//~ End FSceneViewExtensionBase Interface
 
-	/**
-	* Prepare the scene view extension with relevant camera and render mode info.
-	*/
-	void PrepareRender(TSoftObjectPtr<UCineCameraComponent> InCineCameraComponent, bool bInFollowSceneCaptureRenderPath)
-	{
-		CineCameraComponent = InCineCameraComponent;
-		bFollowSceneCaptureRenderPath = bInFollowSceneCaptureRenderPath;
-	};
-
-	/**
+	/** 
 	* Delta time is needed for the purposes of camera smoothing.
 	*/
 	void SetFrameDeltaTime(float InDeltaTime)
@@ -130,30 +136,20 @@ public:
 		return MoveTemp(Output);
 	}
 
-	void SetDisplayConfiguration(const FOpenColorIODisplayConfiguration& InDisplayConfiguration) { DisplayConfiguration = InDisplayConfiguration; };
-
 private:
 	/**
-	* Translates to View.bIsSceneCapture.
+	* A transient property that is used to deliver settings from CineCaptureComponent to views that are related to cine capture.
 	*/
-	bool bFollowSceneCaptureRenderPath = true;
-
-	/**
-	* A transient property that is used to deliver settings from CineCameraComponent to views that are related to cine capture.
-	*/
-	TSoftObjectPtr<UCineCameraComponent> CineCameraComponent = nullptr;
+	TWeakObjectPtr<UCineCaptureComponent2D> CineCaptureComponentWeak;
 
 	/**
 	* Delta time between frames. Used for camera smoothing.
 	*/
-	float DeltaTime = 0.0f;
+	float DeltaTime;
 
 	/** Cached pass resources required to apply conversion for render thread */
 	FOpenColorIORenderPassResources CachedResourcesRenderThread;
 
-	/** Configuration to apply during post render callback */
-	FOpenColorIODisplayConfiguration DisplayConfiguration;
-
 	/** Indicates if OCIO is enabled. */
 	std::atomic<bool> bIsOcioEnabledRenderThread;
 };
@@ -183,8 +179,6 @@ void UCineCaptureComponent2D::UpdateSceneCaptureContents(FSceneInterface* Scene)
 	bOverride_CustomNearClippingPlane = CineCameraComponent->bOverride_CustomNearClippingPlane;
 	CustomNearClippingPlane = CineCameraComponent->CustomNearClippingPlane;
 
-	CineCaptureSVE->PrepareRender(CineCameraComponent, bFollowSceneCaptureRenderPath);
-
 	Scene->UpdateSceneCaptureContents(this);
 }
 
@@ -228,12 +222,10 @@ void UCineCaptureComponent2D::OnRegister()
 	if (!CineCaptureSVE.IsValid())
 	{
 		// This will create a new scene extension per each capture.
-		CineCaptureSVE = MakeShared<FCineCameraCaptureSceneViewExtension>();
+		CineCaptureSVE = MakeShared<FCineCameraCaptureSceneViewExtension>(this);
 		SceneViewExtensions.Add(CineCaptureSVE);
 	}
 
-	CineCaptureSVE->SetDisplayConfiguration(OCIOConfiguration);
-
 	ValidateCineCameraComponent();
 }
 
diff --git a/Engine/Plugins/Experimental/CineCameraSceneCapture/Source/CineCameraSceneCapture/Public/CineCameraSceneCaptureComponent.h b/Engine/Plugins/Experimental/CineCameraSceneCapture/Source/CineCameraSceneCapture/Public/CineCameraSceneCaptureComponent.h
--- a/Engine/Plugins/Experimental/CineCameraSceneCapture/Source/CineCameraSceneCapture/Public/CineCameraSceneCaptureComponent.h
+++ b/Engine/Plugins/Experimental/CineCameraSceneCapture/Source/CineCameraSceneCapture/Public/CineCameraSceneCaptureComponent.h
@@ -37,6 +37,14 @@ public:
 	UPROPERTY(EditAnywhere, BlueprintReadWrite, Category = "Capture Settings", AdvancedDisplay)
 	bool bFollowSceneCaptureRenderPath;
 
+	/** Override state for per-view user flags accessible in materials via TestPostVolumeUserFlag node. */
+	UPROPERTY(EditAnywhere, BlueprintReadWrite, Category = "Capture Settings", AdvancedDisplay)
+	bool bOverrideUserFlags;
+
+	/** Per-view user flags accessible in materials via TestPostVolumeUserFlag node, allowing per-view overrides of material behavior. */
+	UPROPERTY(EditAnywhere, BlueprintReadWrite, Category = "Capture Settings", AdvancedDisplay, meta = (EditCondition = "bOverrideUserFlags"))
+	int32 UserFlags;
+
 	/** 
 	* Open Color IO Configuration. Used to apply color transforms.
 	*/
