CineCameraCapture: Fixed missing camera property copies, which accidentally broke in CL 37799633.

diff --git a/Engine/Plugins/Experimental/CineCameraSceneCapture/Source/CineCameraSceneCapture/Private/CineCameraSceneCaptureComponent.cpp b/Engine/Plugins/Experimental/CineCameraSceneCapture/Source/CineCameraSceneCapture/Private/CineCameraSceneCaptureComponent.cpp
--- a/Engine/Plugins/Experimental/CineCameraSceneCapture/Source/CineCameraSceneCapture/Private/CineCameraSceneCaptureComponent.cpp
+++ b/Engine/Plugins/Experimental/CineCameraSceneCapture/Source/CineCameraSceneCapture/Private/CineCameraSceneCaptureComponent.cpp
@@ -43,12 +43,17 @@ public:
 			return;
 		}
 
+		UCineCameraComponent* CineCameraComponent = CineCaptureComponent->CineCameraComponent.Get();
+		if (!IsValid(CineCameraComponent))
+		{
+			return;
+		}
+
 		FMinimalViewInfo DesiredView;
 
 		// Immitate the behaviour of viewports.
-		CineCaptureComponent->GetCameraView(DeltaTime, DesiredView);
-		FTransform Transform = CineCaptureComponent->GetComponentToWorld();
-		FVector ViewLocation = Transform.GetTranslation();
+		CineCameraComponent->GetCameraView(DeltaTime, DesiredView);
+		const FVector ViewLocation = CineCameraComponent->GetComponentToWorld().GetTranslation();
 		InView.StartFinalPostprocessSettings(ViewLocation);
 		if (CineCaptureComponent->bOverrideUserFlags)
 		{
